CL A: AudioChannelAgnosticCore improvements

Split from CL 51557033. Low-risk edits to existing experimental module.
All APIs UE_EXPERIMENTAL, only consumed by Metasound modules.
Includes build fixes (ScratchBuffer.h, ChannelAgnosticTranscoding.h,
ChannelTypeFamily.h pragma placement, ChannelAgnosticType.h pragma wrapping).
<enter description here>

#tests LowLevelTests Audio
#Synced-CL 51594399

diff --git a/Engine/Source/Runtime/AudioChannelAgnosticCore/Private/ChannelAgnostic/ChannelAgnosticTranscoding.cpp b/Engine/Source/Runtime/AudioChannelAgnosticCore/Private/ChannelAgnostic/ChannelAgnosticTranscoding.cpp
--- a/Engine/Source/Runtime/AudioChannelAgnosticCore/Private/ChannelAgnostic/ChannelAgnosticTranscoding.cpp
+++ b/Engine/Source/Runtime/AudioChannelAgnosticCore/Private/ChannelAgnostic/ChannelAgnosticTranscoding.cpp
@@ -36,13 +36,15 @@ namespace Audio
 			{
 			case EChannelTranscodeMethod::ChannelDrop:
 				{
-					return [&InSrcType,&InDstType](TArrayView<const float*> SrcChannels, TArrayView<float*> DstChannels, const int32 NumFrames)
+					const TArray<FDiscreteChannelTypeFamily::FSpeaker> SrcOrder = InSrcType.GetSpeakerOrder();
+					const TArray<FDiscreteChannelTypeFamily::FSpeaker> DstOrder = InDstType.GetSpeakerOrder();
+					const int32 NumSrcChannels = InSrcType.NumChannels();
+					const int32 NumDstChannels = InDstType.NumChannels();
+					return [SrcOrder, DstOrder, NumSrcChannels, NumDstChannels](TArrayView<const float*> SrcChannels, TArrayView<float*> DstChannels, const int32 NumFrames)
 						{
-							checkSlow(DstChannels.Num() == NumFrames);
-							// Copy everything destination wants, and nothing else.
-							const TArray<FDiscreteChannelTypeFamily::FSpeaker>& SrcOrder = InSrcType.GetSpeakerOrder();
-							const TArray<FDiscreteChannelTypeFamily::FSpeaker>& DstOrder = InDstType.GetSpeakerOrder();
-							const int32 NumDstChannels = InDstType.NumChannels();	
+							checkSlow(SrcChannels.Num() == NumSrcChannels);
+							checkSlow(DstChannels.Num() == NumDstChannels);
+							// Copy everything destination wants, and nothing else.	
 							for (int32 i = 0; i < NumDstChannels; ++i)
 							{
 								const FName DstSpeaker = DstOrder[i].ID;
@@ -131,7 +133,7 @@ namespace Audio
 				const FDiscreteChannelTypeFamily::FSpeaker& Speaker = Order[i];
 			
 				TArray<float>& GainArray = AllChannelGainArrays[i];
-				GainArray.SetNum(InFromType.NumChannels());
+				GainArray.SetNumZeroed(InFromType.NumChannels());
 
 				// Skip LFE.
 				if (Order[i].ID == LexToName(ESpeakerShortNames::LFE))
@@ -176,8 +178,8 @@ namespace Audio
 			const int32 NumChannelsToCopy = FMath::Min(InToType.NumChannels(), InFromType.NumChannels());
 			return [NumChannelsToCopy](TArrayView<const float*> InSrc, TArrayView<float*> InDst, const int32 NumFrames)
 				{
-					checkSlow(InSrc.Num() <= NumChannelsToCopy);
-					checkSlow(InDst.Num() <= NumChannelsToCopy);
+					checkSlow(InSrc.Num() >= NumChannelsToCopy);
+					checkSlow(InDst.Num() >= NumChannelsToCopy);
 					// Copy as many channels as destination will hold.
 					for (int32 i = 0; i < NumChannelsToCopy; ++i)
 					{
@@ -259,11 +261,17 @@ namespace Audio
 			const TArrayView<const ESpeakerShortNames>& InNewOrder,
 			TArray<float>& OutGains)
 		{
+			check(InOldOrder.Num() >= InNumInputChannels && InOldOrder.Num() >= InNumOutputChannels);
 			// Channels not found in new order will be zeroed.
 			OutGains.SetNumZeroed(InNumOutputChannels*InNumInputChannels);
 		
 			// Lookup (old->new).
-			int32 Lookup[static_cast<int32>(ESpeakerShortNames::NumChannels)] = { INDEX_NONE };
+			int32 Lookup[static_cast<int32>(ESpeakerShortNames::NumChannels)];
+			for (int32& L : Lookup) 
+			{ 
+				L = INDEX_NONE; 
+			}
+			
 			for (int32 i = 0; i < InOldOrder.Num(); ++i)
 			{
 				Lookup[static_cast<int32>(InOldOrder[i])] = InNewOrder.IndexOfByKey(InOldOrder[i]);
@@ -300,35 +308,84 @@ namespace Audio
 	// Static speakers can generate a matrix of speaker gains and cache that to avoid running the panner everytime.
 	bool MakeGainMatrixFrom(const FDiscreteChannelTypeFamily& InSrc, const FDiscreteChannelTypeFamily& InDst, TArray<float>& OutMatrix)
 	{
+		// Example: Stereo -> Mono. 
+		// Receiving VBAP, has only a single speaker FC and thus returns 100% gain.
+		// IN   OUT
+		//      FC 
+		// FL  [ 0 ]
+		// FR  [ 0 ]
+				
+		TSet<FName> MixedChannels;
+		const int32 NumDstChannels = InDst.NumChannels();
+		const int32 NumSrcChannels = InSrc.NumChannels();
+		OutMatrix.SetNumZeroed(NumSrcChannels*NumDstChannels);										
+		const TArray<FDiscreteChannelTypeFamily::FSpeaker>& SrcOrder = InSrc.GetSpeakerOrder();
+
+		// See if there's a panner defined, use that.
+		// Sometimes when there's no spatializable channels this will return false.
 		if (const IDiscretePanner* Panner = ChannelAgnosticTranscoder::ChoosePanner(InSrc, InDst))
 		{
-			const int32 DstChannels = InDst.NumChannels();
-			const int32 SrcChannels = InSrc.NumChannels();
-			OutMatrix.SetNumZeroed(SrcChannels*DstChannels);
-
-			const TArray<FDiscreteChannelTypeFamily::FSpeaker>& Order = InSrc.GetSpeakerOrder();
-		
 			// For each src channel, create row vector of gains for each dst channel.
-			TArray<float> SpeakerGains;
-			SpeakerGains.SetNumZeroed(DstChannels);
-			for (int32 i = 0; i < SrcChannels; ++i)
+			for (int32 SrcIndex = 0; SrcIndex < NumSrcChannels; ++SrcIndex)
 			{
-				IDiscretePanner::FOutputParams Results;
-				IDiscretePanner::FInputParams Params { .AzimuthDegrees = Order[i].AzimuthDegrees, .ElevationDegrees = Order[i].ElevationDegrees };
-				if (Panner->ComputeGains(Params,Results) )
+				const FDiscreteChannelTypeFamily::FSpeaker& SrcSpeaker = SrcOrder[SrcIndex];
+				const IDiscretePanner::FInputParams Params { 
+					.AzimuthDegrees = SrcSpeaker.AzimuthDegrees, 
+					.ElevationDegrees = SrcSpeaker.ElevationDegrees,
+					.bAllowAzimuthMirroring = true
+				};
+				IDiscretePanner::FOutputParams Output;
+				if (Panner->ComputeGains(Params,Output) )
 				{
+					const int32 RowOffset = SrcIndex * NumDstChannels;
+					if (Output.Results.Num() > 0)
+					{
+						MixedChannels.Add(SrcSpeaker.ID);
+					}
+
 					// Copy results into gain array.
-					for (int32 j = 0; j < Results.Results.Num(); ++j)
+					for (int32 i = 0; i < Output.Results.Num(); ++i)
 					{
-						if (const int32 DstChannelIndex = Results.Results[j].ChannelIndex; DstChannelIndex != INDEX_NONE)
+						const IDiscretePanner::FPanResult& Result = Output.Results[i];
+						if (const int32 DstChannelIndex = Result.ChannelIndex; DstChannelIndex != INDEX_NONE)
 						{
-							OutMatrix[i*DstChannels+DstChannelIndex] = Results.Results[j].Gain;
+							OutMatrix[RowOffset + DstChannelIndex] = Result.Gain;
 						}
 					}
 				}
 			}
-			return true;
 		}
-		return false;
+		
+		// Non-spatialized channels will need to be copied over verbatim. (LFE, Centre).
+		if (MixedChannels.Num() != InSrc.NumChannels())
+		{
+			for (int32 SrcIndex = 0; SrcIndex < NumSrcChannels; ++SrcIndex) 
+			{
+				const int32 RowOffset = SrcIndex * NumDstChannels;						
+				const FDiscreteChannelTypeFamily::FSpeaker& SrcSpeaker = SrcOrder[SrcIndex];
+				if (MixedChannels.Contains(SrcSpeaker.ID))
+				{
+					continue;
+				}
+			
+				// Not spatialized at SOURCE and didn't mix it as part of the above panner. 
+				if (!SrcSpeaker.bIsSpatialized)
+				{
+					// Find destination match or alias.
+					if (const int32 IndexOf = InDst.GetSpeakerOrder().IndexOfByPredicate(
+						[&SrcSpeaker](const FDiscreteChannelTypeFamily::FSpeaker& In) -> bool
+							{
+								return In.ID == SrcSpeaker.ID;
+							}); IndexOf != INDEX_NONE)
+					{
+						OutMatrix[RowOffset + IndexOf] = 1.0f;	
+						MixedChannels.Add(SrcSpeaker.ID);
+					}
+				}
+			}
+		}
+		
+		// A non-empty result is usable. 
+		return !MixedChannels.IsEmpty();
 	}
 }
\ No newline at end of file
diff --git a/Engine/Source/Runtime/AudioChannelAgnosticCore/Private/ChannelAgnostic/ChannelAgnosticType.cpp b/Engine/Source/Runtime/AudioChannelAgnosticCore/Private/ChannelAgnostic/ChannelAgnosticType.cpp
--- a/Engine/Source/Runtime/AudioChannelAgnosticCore/Private/ChannelAgnostic/ChannelAgnosticType.cpp
+++ b/Engine/Source/Runtime/AudioChannelAgnosticCore/Private/ChannelAgnostic/ChannelAgnosticType.cpp
@@ -20,7 +20,7 @@ namespace Audio
 	
 	FChannelAgnosticType::FChannelAgnosticType(const TRetainedRef<const FChannelTypeFamily> InType, const int32 InNumFrames, const int32 InNumChannels, FSimpleAllocBase* InAllocator)
 		: Buffer(InNumFrames * InNumChannels, InAllocator)
-		, Type(&InType.Get())
+		, Type(&InType.Get()) // FChannelTypeFamily objects are registry-owned with static lifetime; raw pointer is safe and enables copyability.
 		, NumFramesPrivate(InNumFrames)
 		, NumChannelsPrivate(InNumChannels)
 	{
@@ -42,7 +42,7 @@ namespace Audio
 		return Type->IsA(InTypeName);
 	}
 	
-	FName FChannelAgnosticType::GetTypeName() const
+	const FName& FChannelAgnosticType::GetTypeName() const
 	{
 		check(Type);
 		return Type->GetName();
diff --git a/Engine/Source/Runtime/AudioChannelAgnosticCore/Private/TypeFamily/ChannelTypeFamily.cpp b/Engine/Source/Runtime/AudioChannelAgnosticCore/Private/TypeFamily/ChannelTypeFamily.cpp
--- a/Engine/Source/Runtime/AudioChannelAgnosticCore/Private/TypeFamily/ChannelTypeFamily.cpp
+++ b/Engine/Source/Runtime/AudioChannelAgnosticCore/Private/TypeFamily/ChannelTypeFamily.cpp
@@ -95,6 +95,7 @@ namespace Audio
 				
 				// Rebuild lookups.
 				BuildDiscreteBitmaskLookupMap();
+				return true;
 			}
 			return false;
 		}
@@ -301,6 +302,14 @@ namespace Audio
 				0, BaseCat, TEXT("Soundfield Format"),
 				false,
 				true));
+		
+		Registry.RegisterType(
+		TEXT("Cat:Composite"),
+		MakeUnique<FCompositeChannelTypeFamily>(
+			TEXT("Cat:Composite"),
+			BaseCat,
+			TEXT("Composite Format"),
+			false));
 	}
 	
 	void UnregisterCatBaseClassesWithChannelRegistry(IChannelTypeRegistry& Registry)
@@ -309,6 +318,7 @@ namespace Audio
 		Registry.UnregisterType(TEXT("Cat"));
 		Registry.UnregisterType(TEXT("Cat:Discrete"));
 		Registry.UnregisterType(TEXT("Cat:Soundfield"));
+		Registry.UnregisterType(TEXT("Cat:Composite"));
 	}
 	
 }
diff --git a/Engine/Source/Runtime/AudioChannelAgnosticCore/Public/ChannelAgnostic/ChannelAgnosticTranscoding.h b/Engine/Source/Runtime/AudioChannelAgnosticCore/Public/ChannelAgnostic/ChannelAgnosticTranscoding.h
--- a/Engine/Source/Runtime/AudioChannelAgnosticCore/Public/ChannelAgnostic/ChannelAgnosticTranscoding.h
+++ b/Engine/Source/Runtime/AudioChannelAgnosticCore/Public/ChannelAgnostic/ChannelAgnosticTranscoding.h
@@ -46,8 +46,7 @@ namespace Audio
 	}
 		
 	template<typename TFromType>
-	UE_EXPERIMENTAL(5.8,"AudioChannelAgnosticCore is experimental and in flux")
-	struct TFromVisitor final : IChannelTypeVisitor
+	struct UE_EXPERIMENTAL(5.8,"AudioChannelAgnosticCore is experimental and in flux") TFromVisitor final : IChannelTypeVisitor
 	{
 		const TFromType& FromType;
 		const ChannelAgnosticTranscoder::FGetTranscoderParams& Params;
diff --git a/Engine/Source/Runtime/AudioChannelAgnosticCore/Public/ChannelAgnostic/ChannelAgnosticType.h b/Engine/Source/Runtime/AudioChannelAgnosticCore/Public/ChannelAgnostic/ChannelAgnosticType.h
--- a/Engine/Source/Runtime/AudioChannelAgnosticCore/Public/ChannelAgnostic/ChannelAgnosticType.h
+++ b/Engine/Source/Runtime/AudioChannelAgnosticCore/Public/ChannelAgnostic/ChannelAgnosticType.h
@@ -13,12 +13,14 @@ namespace Audio
 	class FChannelTypeFamily;
 	class FSimpleAllocBase;
 
+	PRAGMA_DISABLE_EXPERIMENTAL_WARNINGS
+	
 	class FChannelAgnosticType
 	{
 	public:
 		UE_EXPERIMENTAL(5.8,"AudioChannelAgnosticCore is experimental and in flux")
 		AUDIOCHANNELAGNOSTICCORE_API static FSimpleAllocBase& GetDefaultAllocator();
-
+	
 		UE_EXPERIMENTAL(5.8,"AudioChannelAgnosticCore is experimental and in flux")
 		AUDIOCHANNELAGNOSTICCORE_API explicit FChannelAgnosticType(const TRetainedRef<const FChannelTypeFamily> InType, const int32 InNumFrames, FSimpleAllocBase* InAllocator = &GetDefaultAllocator());
 
@@ -50,21 +52,17 @@ namespace Audio
 		[[nodiscard]]
 		TArrayView<const float> GetChannel(const int32 InChannelIndex) const
 		{
-			PRAGMA_DISABLE_EXPERIMENTAL_WARNINGS					
 			check(InChannelIndex >= 0 && InChannelIndex < NumChannelsPrivate);
 			const TArrayView<const float> AllChannels = Buffer.GetView();
 			return MakeArrayView(AllChannels.GetData() + (NumFramesPrivate * InChannelIndex), NumFramesPrivate);
-			PRAGMA_ENABLE_EXPERIMENTAL_WARNINGS					
 		}
 		
 		UE_EXPERIMENTAL(5.8,"AudioChannelAgnosticCore is experimental and in flux")
 		[[nodiscard]]
 		TArrayView<float> GetChannel(const int32 InChannelIndex)
 		{
-			PRAGMA_DISABLE_EXPERIMENTAL_WARNINGS			
 			const TArrayView<const float> ConstView = std::as_const(*this).GetChannel(InChannelIndex);
 			return MakeArrayView<float>(const_cast<float*>(ConstView.GetData()), ConstView.Num());
-			PRAGMA_ENABLE_EXPERIMENTAL_WARNINGS			
 		}
 
 		UE_EXPERIMENTAL(5.8,"AudioChannelAgnosticCore is experimental and in flux")
@@ -72,7 +70,7 @@ namespace Audio
 		UE_EXPERIMENTAL(5.8,"AudioChannelAgnosticCore is experimental and in flux")
 		AUDIOCHANNELAGNOSTICCORE_API bool IsA(const FName& InTypeName) const;
 		UE_EXPERIMENTAL(5.8,"AudioChannelAgnosticCore is experimental and in flux")
-		AUDIOCHANNELAGNOSTICCORE_API FName GetTypeName() const;
+		AUDIOCHANNELAGNOSTICCORE_API const FName& GetTypeName() const;
 		UE_EXPERIMENTAL(5.8,"AudioChannelAgnosticCore is experimental and in flux")
 		const FChannelTypeFamily& GetType() const { return *Type; }
 
@@ -110,4 +108,6 @@ namespace Audio
 		int32 NumFramesPrivate = 0;
 		int32 NumChannelsPrivate = 0;
 	};
+
+	PRAGMA_ENABLE_EXPERIMENTAL_WARNINGS
 }
diff --git a/Engine/Source/Runtime/AudioChannelAgnosticCore/Public/ScratchBuffer.h b/Engine/Source/Runtime/AudioChannelAgnosticCore/Public/ScratchBuffer.h
--- a/Engine/Source/Runtime/AudioChannelAgnosticCore/Public/ScratchBuffer.h
+++ b/Engine/Source/Runtime/AudioChannelAgnosticCore/Public/ScratchBuffer.h
@@ -7,8 +7,7 @@
 namespace Audio
 {
 	template <typename T>
-	UE_EXPERIMENTAL(5.8,"AudioChannelAgnosticCore is experimental and in flux")
-	class TScratchBuffer
+	class UE_EXPERIMENTAL(5.8,"AudioChannelAgnosticCore is experimental and in flux") TScratchBuffer
 	{
 	public:
 		PRAGMA_DISABLE_EXPERIMENTAL_WARNINGS
diff --git a/Engine/Source/Runtime/AudioChannelAgnosticCore/Public/TypeFamily/ChannelTypeFamily.h b/Engine/Source/Runtime/AudioChannelAgnosticCore/Public/TypeFamily/ChannelTypeFamily.h
--- a/Engine/Source/Runtime/AudioChannelAgnosticCore/Public/TypeFamily/ChannelTypeFamily.h
+++ b/Engine/Source/Runtime/AudioChannelAgnosticCore/Public/TypeFamily/ChannelTypeFamily.h
@@ -335,15 +335,15 @@ namespace Audio
 		TUniquePtr<IPanner> Panner;	// Just one for now.
 	};
 
-	// Soundfield 
-	class FSoundfieldChannelTypeFamily : public FChannelTypeFamily 
+	// Soundfield
+	PRAGMA_DISABLE_EXPERIMENTAL_WARNINGS
+	class FSoundfieldChannelTypeFamily : public FChannelTypeFamily
 	{
 	public:
 		UE_NONCOPYABLE(FSoundfieldChannelTypeFamily);
 		using Super = FChannelTypeFamily;
-		
-		PRAGMA_DISABLE_EXPERIMENTAL_WARNINGS
-		
+
+
 		UE_EXPERIMENTAL(5.8,"AudioChannelAgnosticCore is experimental and in flux")
 		static FName GetFamilyTypeName()
 		{
@@ -372,20 +372,19 @@ namespace Audio
 			InVisitor.Visit(*this);
 		}
 		
-		PRAGMA_ENABLE_EXPERIMENTAL_WARNINGS		
 	private:
 		int32 Order = 0;
 	};
+	PRAGMA_ENABLE_EXPERIMENTAL_WARNINGS
 
-	
+	PRAGMA_DISABLE_EXPERIMENTAL_WARNINGS
 	class FCompositeChannelTypeFamily : public FChannelTypeFamily
 	{
 	public:
 		UE_NONCOPYABLE(FCompositeChannelTypeFamily);
 
 		using Super = FChannelTypeFamily;
-		
-		PRAGMA_DISABLE_EXPERIMENTAL_WARNINGS		
+
 
 		UE_EXPERIMENTAL(5.8,"AudioChannelAgnosticCore is experimental and in flux")
 		static FName GetFamilyTypeName()
@@ -395,8 +394,9 @@ namespace Audio
 		}
 		UE_EXPERIMENTAL(5.8,"AudioChannelAgnosticCore is experimental and in flux")
 		explicit FCompositeChannelTypeFamily(const FName& InUniqueName, FChannelTypeFamily* InParentType, const FString& InFriendlyName, const bool bIsParentsDefault)
-			: Super(InUniqueName, GetFamilyTypeName(), 0, InParentType, InFriendlyName, bIsParentsDefault, false)
-		{}
+			: Super(InUniqueName, GetFamilyTypeName(), /* Placeholder, NumChannels=1 */ 1, InParentType, InFriendlyName, bIsParentsDefault, false)
+		{
+		}
 	
 		UE_EXPERIMENTAL(5.8,"AudioChannelAgnosticCore is experimental and in flux")
 		int32 NumTypes() const
@@ -427,18 +427,17 @@ namespace Audio
 			InVisitor.Visit(*this);
 		}
 		
-		PRAGMA_ENABLE_EXPERIMENTAL_WARNINGS
-		
 	protected:
 		TArray<FContainedType> Types;
 	};
-	
-	
+	PRAGMA_ENABLE_EXPERIMENTAL_WARNINGS
+
+
 	// Channel type registry.
+	PRAGMA_DISABLE_EXPERIMENTAL_WARNINGS
 	class IChannelTypeRegistry : public ITypeFamilyRegistry
 	{
 	public:
-		PRAGMA_DISABLE_EXPERIMENTAL_WARNINGS				
 		
 		UE_EXPERIMENTAL(5.8,"AudioChannelAgnosticCore is experimental and in flux")			
 		const FChannelTypeFamily* FindChannel(const FName InName) const
@@ -492,10 +491,10 @@ namespace Audio
 		UE_EXPERIMENTAL(5.8,"AudioChannelAgnosticCore is experimental and in flux")	
 		virtual TArray<TSharedRef<const FChannelTypeFamily>> FindAllFamily(const FName InFamilyName) const = 0;		
 
-		PRAGMA_ENABLE_EXPERIMENTAL_WARNINGS
 	};
+	PRAGMA_ENABLE_EXPERIMENTAL_WARNINGS
 
-	UE_EXPERIMENTAL(5.8,"AudioChannelAgnosticCore is experimental and in flux")		
+	UE_EXPERIMENTAL(5.8,"AudioChannelAgnosticCore is experimental and in flux")
 	AUDIOCHANNELAGNOSTICCORE_API IChannelTypeRegistry& GetChannelRegistry();
 	
 	UE_EXPERIMENTAL(5.8,"AudioChannelAgnosticCore is experimental and in flux")	
diff --git a/Engine/Source/Runtime/AudioChannelAgnosticCore/Public/TypeFamily/TypeFamily.h b/Engine/Source/Runtime/AudioChannelAgnosticCore/Public/TypeFamily/TypeFamily.h
--- a/Engine/Source/Runtime/AudioChannelAgnosticCore/Public/TypeFamily/TypeFamily.h
+++ b/Engine/Source/Runtime/AudioChannelAgnosticCore/Public/TypeFamily/TypeFamily.h
@@ -54,9 +54,9 @@ namespace Audio
 			PRAGMA_ENABLE_EXPERIMENTAL_WARNINGS			
 		}
 		UE_EXPERIMENTAL(5.8,"AudioChannelAgnosticCore is experimental and in flux")
-		FName GetName() const { return Name; }
+		const FName& GetName() const { return Name; }
 		UE_EXPERIMENTAL(5.8,"AudioChannelAgnosticCore is experimental and in flux")
-		FString GetFriendlyName() const { return FriendlyName; }
+		const FString& GetFriendlyName() const { return FriendlyName; }
 		UE_EXPERIMENTAL(5.8,"AudioChannelAgnosticCore is experimental and in flux")
 		const FTypeFamily* GetParent() const { return Parent; }
 	};
