[Backout] - CL51661203
#rnx
Original CL Desc
-----------------------------------------------------------------
Add FindDiscreteFormatFromNumChannels utility to FChannelAgnosticUtils for looking up discrete channel type families by channel count
Add Metasound channel-agnostic auto format support and wire up AudioChannelAgnosticCore dependency
Add engine and audio unit tests for the new functionality

#tests EngineTest and AudioUnitTests

diff --git a/Engine/Plugins/Runtime/Metasound/Source/MetasoundEngine/MetasoundEngine.Build.cs b/Engine/Plugins/Runtime/Metasound/Source/MetasoundEngine/MetasoundEngine.Build.cs
--- a/Engine/Plugins/Runtime/Metasound/Source/MetasoundEngine/MetasoundEngine.Build.cs
+++ b/Engine/Plugins/Runtime/Metasound/Source/MetasoundEngine/MetasoundEngine.Build.cs
@@ -26,8 +26,7 @@ namespace UnrealBuildTool.Rules
 					"AudioMixer",
 					"MetasoundGraphCore",
 					"MetasoundGenerator",
-					"SignalProcessing",
-					"AudioChannelAgnosticCore"
+					"SignalProcessing"
 				}
 			);
 
diff --git a/Engine/Plugins/Runtime/Metasound/Source/MetasoundEngine/Private/MetasoundChannelAgnosticAutoFormat.cpp b/Engine/Plugins/Runtime/Metasound/Source/MetasoundEngine/Private/MetasoundChannelAgnosticAutoFormat.cpp
deleted file mode 100644
diff --git a/Engine/Plugins/Runtime/Metasound/Source/MetasoundEngine/Public/MetasoundChannelAgnosticAutoFormat.h b/Engine/Plugins/Runtime/Metasound/Source/MetasoundEngine/Public/MetasoundChannelAgnosticAutoFormat.h
deleted file mode 100644
diff --git a/Engine/Plugins/Runtime/Metasound/Source/MetasoundEngineTest/MetasoundEngineTest.Build.cs b/Engine/Plugins/Runtime/Metasound/Source/MetasoundEngineTest/MetasoundEngineTest.Build.cs
--- a/Engine/Plugins/Runtime/Metasound/Source/MetasoundEngineTest/MetasoundEngineTest.Build.cs
+++ b/Engine/Plugins/Runtime/Metasound/Source/MetasoundEngineTest/MetasoundEngineTest.Build.cs
@@ -17,7 +17,6 @@ namespace UnrealBuildTool.Rules
 					"CoreUObject",
 					"Engine",
 					"AudioExtensions",
-					"AudioChannelAgnosticCore",
 					"AudioMixer",
 					"SignalProcessing",
 					"Projects"
diff --git a/Engine/Plugins/Runtime/Metasound/Source/MetasoundEngineTest/Private/EngineTestChannelAgnosticAutoFormat.cpp b/Engine/Plugins/Runtime/Metasound/Source/MetasoundEngineTest/Private/EngineTestChannelAgnosticAutoFormat.cpp
deleted file mode 100644
diff --git a/Engine/Source/Runtime/AudioChannelAgnosticCore/Private/ChannelAgnostic/ChannelAgnosticTypeUtils.cpp b/Engine/Source/Runtime/AudioChannelAgnosticCore/Private/ChannelAgnostic/ChannelAgnosticTypeUtils.cpp
--- a/Engine/Source/Runtime/AudioChannelAgnosticCore/Private/ChannelAgnostic/ChannelAgnosticTypeUtils.cpp
+++ b/Engine/Source/Runtime/AudioChannelAgnosticCore/Private/ChannelAgnostic/ChannelAgnosticTypeUtils.cpp
@@ -56,14 +56,4 @@ namespace Audio
 		Audio::Interleave(In.Buffer.GetView(), In.NumChannels(), Out);
 	}
 
-	TSharedPtr<const FDiscreteChannelTypeFamily> FChannelAgnosticUtils::FindDiscreteFormatFromNumChannels(const int32 InNumChannels)
-	{
-		const TOptional<uint32> Mask = FindChannelMaskFromAudioMixerNumChannels(InNumChannels);
-		if (!Mask.IsSet())
-		{
-			return {};
-		}
-		return GetChannelRegistry().FindDiscreteChannelFromBitMask(*Mask);
-	}
-
 }
\ No newline at end of file
diff --git a/Engine/Source/Runtime/AudioChannelAgnosticCore/Public/ChannelAgnostic/ChannelAgnosticTypeUtils.h b/Engine/Source/Runtime/AudioChannelAgnosticCore/Public/ChannelAgnostic/ChannelAgnosticTypeUtils.h
--- a/Engine/Source/Runtime/AudioChannelAgnosticCore/Public/ChannelAgnostic/ChannelAgnosticTypeUtils.h
+++ b/Engine/Source/Runtime/AudioChannelAgnosticCore/Public/ChannelAgnostic/ChannelAgnosticTypeUtils.h
@@ -24,15 +24,7 @@ namespace Audio
 		UE_EXPERIMENTAL(5.8,"AudioChannelAgnosticCore is experimental and in flux")
 		AUDIOCHANNELAGNOSTICCORE_API static bool ChannelIdToShortMixerChannels(const TArray<FName>& InIds, TArray<ESpeakerShortNames>& OutShortNames);
 				
-		UE_EXPERIMENTAL(5.8,"AudioChannelAgnosticCore is experimental and in flux")
+		UE_EXPERIMENTAL(5.8,"AudioChannelAgnosticCore is experimental and in flux")		
 		static AUDIOCHANNELAGNOSTICCORE_API void Interleave(const FChannelAgnosticType& In, TArrayView<float> Out);
-
-		/**
-		 * Finds a discrete channel type family matching the given number of channels,
-		 * using the standard AudioMixer channel layouts.
-		 * @return Discrete channel type family, or null if no match found.
-		 */
-		UE_EXPERIMENTAL(5.8,"AudioChannelAgnosticCore is experimental and in flux")
-		AUDIOCHANNELAGNOSTICCORE_API static TSharedPtr<const FDiscreteChannelTypeFamily> FindDiscreteFormatFromNumChannels(const int32 InNumChannels);
 	};
 }
